Move the widget carousel and it's resources to a stand along module so that it can be used in both the launcher and the editor.

diff --git a/Engine/Source/Developer/WidgetCarousel/Private/WidgetCarouselPrivatePCH.h b/Engine/Source/Developer/WidgetCarousel/Private/WidgetCarouselPrivatePCH.h
new file mode 100644
--- /dev/null
+++ b/Engine/Source/Developer/WidgetCarousel/Private/WidgetCarouselPrivatePCH.h
@@ -0,0 +1,13 @@
+// Copyright 1998-2014 Epic Games, Inc. All Rights Reserved.
+
+#pragma once
+
+#include "SlateBasics.h"
+
+#include "WidgetCarouselPrivatePCH.h"
+
+#include "WidgetCarouselModule.h"
+#include "SWidgetCarousel.h"
+#include "SCarouselNavigationButton.h"
+#include "SCarouselNavigationBar.h"
+#include "WidgetCarouselStyle.h"
diff --git a/Engine/Source/Developer/WidgetCarousel/Private/WidgetCarouselStyle.cpp b/Engine/Source/Developer/WidgetCarousel/Private/WidgetCarouselStyle.cpp
new file mode 100644
--- /dev/null
+++ b/Engine/Source/Developer/WidgetCarousel/Private/WidgetCarouselStyle.cpp
@@ -0,0 +1,101 @@
+// Copyright 1998-2014 Epic Games, Inc. All Rights Reserved.
+
+#include "WidgetCarouselPrivatePCH.h"
+
+#include "EditorStyle.h"
+
+TSharedPtr< FSlateStyleSet > FWidgetCarouselStyle::WidgetCarouselStyleInstance = NULL;
+
+void FWidgetCarouselStyle::Initialize()
+{
+	if ( !WidgetCarouselStyleInstance.IsValid() )
+	{
+		WidgetCarouselStyleInstance = Create();
+		FSlateStyleRegistry::RegisterSlateStyle( *WidgetCarouselStyleInstance );
+	}
+}
+
+void FWidgetCarouselStyle::Shutdown()
+{
+	FSlateStyleRegistry::UnRegisterSlateStyle( *WidgetCarouselStyleInstance );
+	ensure( WidgetCarouselStyleInstance.IsUnique() );
+	WidgetCarouselStyleInstance.Reset();
+}
+
+FName FWidgetCarouselStyle::GetStyleSetName()
+{
+	static FName StyleSetName(TEXT("WidgetCarouselStyle"));
+	return StyleSetName;
+}
+
+#define IMAGE_BRUSH( RelativePath, ... ) FSlateImageBrush( Style->RootToContentDir( RelativePath, TEXT(".png") ), __VA_ARGS__ )
+#define BOX_BRUSH( RelativePath, ... ) FSlateBoxBrush( Style->RootToContentDir( RelativePath, TEXT(".png") ), __VA_ARGS__ )
+#define BORDER_BRUSH( RelativePath, ... ) FSlateBorderBrush( Style->RootToContentDir( RelativePath, TEXT(".png") ), __VA_ARGS__ )
+#define TTF_FONT( RelativePath, ... ) FSlateFontInfo( Style->RootToContentDir( RelativePath, TEXT(".ttf") ), __VA_ARGS__ )
+#define OTF_FONT( RelativePath, ... ) FSlateFontInfo( Style->RootToContentDir( RelativePath, TEXT(".otf") ), __VA_ARGS__ )
+
+const FVector2D Icon16x16(16.0f, 16.0f);
+const FVector2D Icon20x20(20.0f, 20.0f);
+const FVector2D Icon40x40(40.0f, 40.0f);
+
+TSharedRef< FSlateStyleSet > FWidgetCarouselStyle::Create()
+{
+	TSharedRef< FSlateStyleSet > Style = MakeShareable(new FSlateStyleSet("WidgetCarouselStyle"));
+	Style->SetContentRoot(FPaths::EngineContentDir() / TEXT("Editor/Slate/WidgetCarousel"));
+
+	const FButtonStyle DefaultButton = FButtonStyle()
+		.SetNormalPadding(0)
+		.SetPressedPadding(0);
+
+	FLinearColor PrimaryCallToActionColor = FLinearColor(1.0, 0.7372, 0.05637);
+	FLinearColor PrimaryCallToActionColor_Hovered = FLinearColor(1.0, 0.83553, 0.28445);
+	FLinearColor PrimaryCallToActionColor_Pressed = FLinearColor(1.0, 0.66612, 0.0012);
+
+	Style->Set("PrimaryColor", PrimaryCallToActionColor);
+
+	Style->Set("WidgetCarousel.NavigationButtonLeft", new IMAGE_BRUSH("Arrow-Left", FVector2D(25.0f, 42.0f)));
+
+	Style->Set("WidgetCarousel.NavigationButtonRight", new IMAGE_BRUSH("Arrow-Right", FVector2D(25.0f, 42.0f)));
+
+	Style->Set("WidgetCarousel.NavigationButton", FButtonStyle(DefaultButton)
+		.SetNormal(BOX_BRUSH("WhiteBox_7px_CornerRadius", FVector2D(17, 17), FMargin(0.5), PrimaryCallToActionColor))
+		.SetPressed(BOX_BRUSH("WhiteBox_7px_CornerRadius", FVector2D(17, 17), FMargin(0.5), PrimaryCallToActionColor_Pressed))
+		.SetHovered(BOX_BRUSH("WhiteBox_7px_CornerRadius", FVector2D(17, 17), FMargin(0.5), PrimaryCallToActionColor_Hovered)));
+
+	Style->Set("WidgetCarousel.NavigationScrollBarHighlight", new IMAGE_BRUSH("CarouselNavMarker", FVector2D(80.0f, 20.0f), PrimaryCallToActionColor));
+
+	Style->Set("WidgetCarousel.NavigationScrollBarLeftButton", FButtonStyle(DefaultButton)
+		.SetNormal(IMAGE_BRUSH("CarouselNavLeft", FVector2D(80.0f, 20.0f), FLinearColor::White))
+		.SetHovered(IMAGE_BRUSH("CarouselNavLeft", FVector2D(80.0f, 20.0f), PrimaryCallToActionColor))
+		.SetPressed(IMAGE_BRUSH("CarouselNavLeft", FVector2D(80.0f, 20.0f), PrimaryCallToActionColor_Pressed)));
+
+	Style->Set("WidgetCarousel.NavigationScrollBarCenterButton", FButtonStyle(DefaultButton)
+		.SetNormal(IMAGE_BRUSH("CarouselNavCenter", FVector2D(80.0f, 20.0f), FLinearColor::White))
+		.SetHovered(IMAGE_BRUSH("CarouselNavCenter", FVector2D(80.0f, 20.0f), PrimaryCallToActionColor))
+		.SetPressed(IMAGE_BRUSH("CarouselNavCenter", FVector2D(80.0f, 20.0f), PrimaryCallToActionColor_Pressed)));
+
+	Style->Set("WidgetCarousel.NavigationScrollBarRightButton", FButtonStyle(DefaultButton)
+		.SetNormal(IMAGE_BRUSH("CarouselNavRight", FVector2D(80.0f, 20.0f), FLinearColor::White))
+		.SetHovered(IMAGE_BRUSH("CarouselNavRight", FVector2D(80.0f, 20.0f), PrimaryCallToActionColor))
+		.SetPressed(IMAGE_BRUSH("CarouselNavRight", FVector2D(80.0f, 20.0f), PrimaryCallToActionColor_Pressed)));
+
+	Style->Set("Black .6", FLinearColor(0, 0, 0, 0.6f));
+
+	return Style;
+}
+
+#undef IMAGE_BRUSH
+#undef BOX_BRUSH
+#undef BORDER_BRUSH
+#undef TTF_FONT
+#undef OTF_FONT
+
+void FWidgetCarouselStyle::ReloadTextures()
+{
+	FSlateApplication::Get().GetRenderer()->ReloadTextureResources();
+}
+
+const ISlateStyle& FWidgetCarouselStyle::Get()
+{
+	return *WidgetCarouselStyleInstance;
+}
diff --git a/Engine/Source/Developer/WidgetCarousel/Public/SCarouselNavigationBar.h b/Engine/Source/Developer/WidgetCarousel/Public/SCarouselNavigationBar.h
new file mode 100644
--- /dev/null
+++ b/Engine/Source/Developer/WidgetCarousel/Public/SCarouselNavigationBar.h
@@ -0,0 +1,54 @@
+// Copyright 1998-2014 Epic Games, Inc. All Rights Reserved.
+#pragma once
+
+#define LOCTEXT_NAMESPACE "WidgetCarousel"
+
+/**
+* A horizontal bar of buttons for navigating to a specific item in the widget carousel.
+*/
+class WIDGETCAROUSEL_API SCarouselNavigationBar : public SCompoundWidget
+{
+	DECLARE_DELEGATE_OneParam(FOnSelectedIndexChanged, int32)
+
+public:
+	SLATE_BEGIN_ARGS(SCarouselNavigationBar)
+		: _OnSelectedIndexChanged()
+		, _ItemCount(0)
+		, _CurrentItemIndex(0)
+		, _CurrentSlideAmount(0.0f)
+	{}
+
+		/** Called when the selected index changes. */
+		SLATE_EVENT(FOnSelectedIndexChanged, OnSelectedIndexChanged)
+
+		/** The number of items to generate navigation buttons for. */
+		SLATE_ARGUMENT(int32, ItemCount)
+
+		/** The index of the currently selected item. */
+		SLATE_ATTRIBUTE(int32, CurrentItemIndex)
+
+		/** The amount of offset which should be applied to the selected item indicator. */
+		SLATE_ATTRIBUTE(float, CurrentSlideAmount)
+
+	SLATE_END_ARGS()
+
+	int32 GetItemCount();
+
+	void SetItemCount(int32 InItemCount);
+
+	void Construct(const FArguments& InArgs);
+
+private:
+	void BuildScrollBar();
+
+	FReply HandleItemButtonClicked(int32 ItemIndex);
+
+	FVector2D GetPlaceHolderPosition() const;
+
+private:
+	TSharedPtr<SHorizontalBox> WidgetScrollBox;
+	FOnSelectedIndexChanged OnSelectedIndexChanged;
+	int32 ItemCount;
+	TAttribute<int32> CurrentItemIndex;
+	TAttribute<float> CurrentSlideAmount;
+};
\ No newline at end of file
diff --git a/Engine/Source/Developer/WidgetCarousel/Public/SCarouselNavigationButton.h b/Engine/Source/Developer/WidgetCarousel/Public/SCarouselNavigationButton.h
new file mode 100644
--- /dev/null
+++ b/Engine/Source/Developer/WidgetCarousel/Public/SCarouselNavigationButton.h
@@ -0,0 +1,252 @@
+// Copyright 1998-2014 Epic Games, Inc. All Rights Reserved.
+#pragma once
+
+#include "SWidgetCarousel.h"
+#include "WidgetCarouselStyle.h"
+
+#define LOCTEXT_NAMESPACE "WidgetCarousel"
+
+/**
+ * A widget used to navigate the carousel widget
+ */
+class SCarouselNavigationButton
+	: public SCompoundWidget
+{
+	DECLARE_DELEGATE(FOnBeginPeek)
+	DECLARE_DELEGATE(FOnEndPeek)
+
+public:
+
+	class SPeekBorder 
+		: public SBorder
+	{
+	public:
+
+		SLATE_BEGIN_ARGS(SPeekBorder)
+			: _OnBeginPeek()
+			, _OnEndPeek()
+			, _HAlign(HAlign_Fill)
+			, _VAlign(VAlign_Fill)
+			, _Padding(FMargin(0))
+		{}
+			/** Called when the button is hovered */
+			SLATE_EVENT(FOnBeginPeek, OnBeginPeek)
+			SLATE_EVENT(FOnEndPeek, OnEndPeek)
+
+			SLATE_ARGUMENT(EHorizontalAlignment, HAlign)
+			SLATE_ARGUMENT(EVerticalAlignment, VAlign)
+			SLATE_ATTRIBUTE(FMargin, Padding)
+
+			SLATE_DEFAULT_SLOT(FArguments, Content)
+
+		SLATE_END_ARGS()
+
+		void Construct(const FArguments& InArgs)
+		{
+			OnBeginPeek = InArgs._OnBeginPeek;
+			OnEndPeek = InArgs._OnEndPeek;
+
+			SBorder::Construct(SBorder::FArguments()
+				.HAlign(InArgs._HAlign)
+				.VAlign(InArgs._VAlign)
+				.Padding(InArgs._Padding)
+				.BorderImage(FStyleDefaults::GetNoBrush())
+				[
+					InArgs._Content.Widget
+				]);
+		}
+
+		void OnMouseEnter(const FGeometry& MyGeometry, const FPointerEvent& MouseEvent)
+		{
+			if (FPlatformProcess::IsThisApplicationForeground())
+			{
+				if (OnBeginPeek.IsBound() == true)
+				{
+					OnBeginPeek.Execute();
+				}
+			}
+			SBorder::OnMouseEnter(MyGeometry, MouseEvent);
+		}
+
+		void OnMouseLeave(const FPointerEvent& MouseEvent)
+		{
+			if (FPlatformProcess::IsThisApplicationForeground())
+			{
+				if (OnEndPeek.IsBound() == true)
+				{
+					OnEndPeek.Execute();
+				}
+			}
+			SBorder::OnMouseLeave(MouseEvent);
+		}
+
+	private:
+
+		FOnBeginPeek OnBeginPeek;
+		FOnEndPeek OnEndPeek;
+	};
+
+public:
+
+	enum class ENavigationButtonDirection
+	{
+		Left,
+		Right
+	};
+
+	SLATE_BEGIN_ARGS(SCarouselNavigationButton) {}
+
+		/** Called when the button is clicked */
+		SLATE_EVENT(FOnClicked, OnClicked)
+
+		/** Specifies which direction should be used to setup the button's visuals and layout. */
+		SLATE_ARGUMENT(ENavigationButtonDirection, Direction)
+
+		/** Called when the carousel should shift it's image to give a preview of the next image */
+		SLATE_EVENT(FOnBeginPeek, OnBeginPeek)
+		SLATE_EVENT(FOnEndPeek, OnEndPeek)
+
+	SLATE_END_ARGS()
+
+	void Construct(const FArguments& InArgs)
+	{
+		ImageTransparency = 0.f;
+
+		ChildSlot
+		[
+			SNew(SBox)
+			.WidthOverride(150.0f)
+			[
+				SNew(SPeekBorder)
+				.OnBeginPeek(InArgs._OnBeginPeek)
+				.OnEndPeek(InArgs._OnEndPeek)
+				[
+					SNew(SHorizontalBox)
+					+ SHorizontalBox::Slot()
+					[
+						SNew(SSpacer)
+						.Visibility(InArgs._Direction == ENavigationButtonDirection::Left ? EVisibility::Collapsed : EVisibility::Visible)
+					]
+					+SHorizontalBox::Slot()
+					.HAlign(InArgs._Direction == ENavigationButtonDirection::Left ? HAlign_Left : HAlign_Right)
+					.AutoWidth()
+					[
+						SAssignNew(HiddenButton, SButton)
+						.HAlign(HAlign_Center)
+						.VAlign(VAlign_Center)
+						.ToolTipText(InArgs._ToolTipText)
+						.OnClicked(InArgs._OnClicked)
+						.ButtonStyle(FCoreStyle::Get(), "NoBorder")
+						.Cursor(EMouseCursor::Hand)
+						[
+							SAssignNew(VisibleButton, SButton)
+							.HAlign(HAlign_Center)
+							.VAlign(VAlign_Center)
+							.ToolTipText(InArgs._ToolTipText)
+							.OnClicked(InArgs._OnClicked)
+							.ButtonStyle(FWidgetCarouselStyle::Get(), "WidgetCarousel.NavigationButton")
+							.ContentPadding(FMargin(15, 30))
+							.ButtonColorAndOpacity(this, &SCarouselNavigationButton::GetButtonColor)
+							.Cursor(EMouseCursor::Hand)
+							[
+								SNew(SBox)
+								.HeightOverride(42.0f)
+								.WidthOverride(25.0f)
+								[
+									SNew(SImage)
+									.Image(InArgs._Direction == ENavigationButtonDirection::Left ?
+										FWidgetCarouselStyle::Get().GetBrush("WidgetCarousel.NavigationButtonLeft") :
+										FWidgetCarouselStyle::Get().GetBrush("WidgetCarousel.NavigationButtonRight"))
+									.ColorAndOpacity(this, &SCarouselNavigationButton::GetButtonImageColor)
+								]
+							]
+						]
+					]
+					+ SHorizontalBox::Slot()
+					[
+						SNew(SSpacer)
+						.Visibility(InArgs._Direction == ENavigationButtonDirection::Left ? EVisibility::Visible : EVisibility::Collapsed)
+					]
+				]
+			]
+		];
+	}
+
+private:
+
+	FSlateColor GetButtonColor() const
+	{
+		if (HiddenButton->IsHovered())
+		{
+			if (!VisibleButton->IsHovered())
+			{
+				FLinearColor Color = FWidgetCarouselStyle::Get().GetColor("PrimaryColor");
+				Color.A = ImageTransparency;
+				return Color;
+			}
+			else
+			{
+				return FLinearColor::White;
+			}
+		}
+
+		return FLinearColor(0.0f, 0.0f, 0.0f, ImageTransparency * 0.5f);
+	}
+
+	FSlateColor GetButtonImageColor() const
+	{
+		if (HiddenButton->IsHovered() || VisibleButton->IsHovered())
+		{
+			if (!VisibleButton->IsHovered())
+			{
+				return FLinearColor(0.0f, 0.0f, 0.0f, 1.0f);
+			}
+
+			return FLinearColor(0.0f, 0.0f, 0.0f, ImageTransparency);
+		}
+
+		return FLinearColor(1.0f, 1.0f, 1.0f, ImageTransparency);
+	}
+
+	virtual void Tick(const FGeometry& AllottedGeometry, const double InCurrentTime, const float InDeltaTime) override
+	{
+		//SButton::Tick(AllottedGeometry, InCurrentTime, InDeltaTime);
+
+		static const float BlendSpeed = 2.0f;
+
+		float DesiredBlendSpeed = BlendSpeed * InDeltaTime;
+
+		if (IsHovered())
+		{
+			if (ImageTransparency != 1.f && FPlatformProcess::IsThisApplicationForeground())
+			{
+				ImageTransparency = FMath::Min<float>(ImageTransparency + DesiredBlendSpeed, 1.f);
+			}
+			else
+			{
+				ImageTransparency = 1.f;
+			}
+		}
+		else
+		{
+			if (ImageTransparency != 0.f && FPlatformProcess::IsThisApplicationForeground())
+			{
+				ImageTransparency = FMath::Max<float>(ImageTransparency - DesiredBlendSpeed, 0.f);
+			}
+			else
+			{
+				ImageTransparency = 0.f;
+			}
+		}
+	}
+
+private:
+
+	TSharedPtr<SButton> HiddenButton;
+	TSharedPtr<SButton> VisibleButton;
+
+	float ImageTransparency;
+	ENavigationButtonDirection Direction;
+};
+
+#undef LOCTEXT_NAMESPACE
diff --git a/Engine/Source/Developer/WidgetCarousel/Public/SWidgetCarousel.h b/Engine/Source/Developer/WidgetCarousel/Public/SWidgetCarousel.h
new file mode 100644
--- /dev/null
+++ b/Engine/Source/Developer/WidgetCarousel/Public/SWidgetCarousel.h
@@ -0,0 +1,724 @@
+// Copyright 1998-2014 Epic Games, Inc. All Rights Reserved.
+
+/*=============================================================================
+	SWidgetCarousel.h: Declares the SWidgetCarousel class.
+=============================================================================*/
+
+#pragma once
+
+/** The desired Carousel scroll direction */
+namespace EWidgetCarouselScrollDirection
+{
+	enum Type
+	{
+		// Scroll the carousel left.
+		Carousel_Left,
+		// Scroll the carousel right.
+		Carousel_Right,
+		// Center the carousel
+		Carousel_Center,
+	};
+}
+
+/** The Carousel display widget that gets animated by the carousel widget */
+struct FCarouselDisplayItem : public TSharedFromThis<FCarouselDisplayItem>
+{
+	/**
+	 * Construct the display widget.
+	 */
+	FCarouselDisplayItem()
+		: OpacityValue( 0.f )
+		, SlideValue( 0.f )
+		, DesiredOpacityValue( 0.f )
+		, SlideValueLeftLimit( -1.0f )
+		, SlideValueRightLimit( 1.0f )
+		, PeakValueLeftLimit(-0.05f)
+		, PeakValueRightLimit(0.05f)
+		, OptimalSlideValue( 0.0f )
+		, DesiredSlideValue( 0.f )
+		, bTransition( false )
+		, bPeak(false)
+		, bFade( false )
+		, MoveSpeed( 5.0f )
+		, PeakSpeed(.2f)
+		, PeakDistance(0.05f)
+		, FadeRate( 2.0f )
+	{
+		// Create default widget
+		FXWidget = SNew(SFxWidget)
+			.IgnoreClipping(false);
+
+		SlideInCurve = FCurveSequence();
+		SlideInCurve.AddCurve(0.1f, 0.5f);
+	}
+
+	/**
+	 * Scroll in the widget.
+	 * @ScrollDirection The scroll direction.
+	 */
+	void ScrollIn( EWidgetCarouselScrollDirection::Type ScrollDirection )
+	{
+		bTransition = true;
+		bPeak = false;
+		bFade = FadeRate != 0;
+		DesiredSlideValue = OptimalSlideValue;
+		DesiredOpacityValue = 1.0f;
+	}
+
+	/**
+	 * Scroll out the widget.
+	 * @param ScrollDirection - The scroll direction.
+	 */
+	void PeakIn(EWidgetCarouselScrollDirection::Type ScrollDirection)
+	{
+		if (!bTransition)
+		{
+			SlideInCurve.JumpToEnd();
+			SlideInCurve.PlayReverse();
+			bPeak = true;
+			bFade = FadeRate != 0;
+			switch (ScrollDirection)
+			{
+				case EWidgetCarouselScrollDirection::Carousel_Left:
+				{
+					DesiredSlideValue = PeakValueLeftLimit;
+				}
+				break;
+				case EWidgetCarouselScrollDirection::Carousel_Right:
+				{
+					DesiredSlideValue = PeakValueRightLimit;
+				}
+				break;
+				case EWidgetCarouselScrollDirection::Carousel_Center:
+				{
+					DesiredSlideValue = OptimalSlideValue;
+				}
+				break;
+			}
+			DesiredOpacityValue = 0.0f;
+		}
+	}
+
+	/**
+	 * Tick the widget - animate the slide / fade.
+	 * @param DeltaTime - The delta time.
+	 */
+	void Tick( float DeltaTime )
+	{
+		if ( bTransition )
+		{
+			bTransition = BlendWidget(DeltaTime, DesiredSlideValue, MoveSpeed, SlideValue);
+			SetSlide( SlideValue );
+		}
+		else if ( bPeak )
+		{
+			bPeak = BlendWidget(DeltaTime, DesiredSlideValue, PeakSpeed, SlideValue);
+			SetSlide(SlideValue);
+		}
+
+		if ( bFade )
+		{
+			bFade = BlendWidget(DeltaTime, DesiredOpacityValue, FadeRate, OpacityValue);
+			SetOpacity(OpacityValue);
+		}
+	}
+
+	/**
+	 * Blend the widget fade / transform to the desired value.
+	 * @param DeltaTime - The delta time.
+	 * @param DesiredValue - The desired value.
+	 * @param Speed - Speed to transform.
+	 * @param OutCurrentValue - The value to set.
+	 * @return True if still blending, false otherwise
+	 */
+	bool BlendWidget( const float& DeltaTime, const float& DesiredValue, const float& Speed, float& OutCurrentValue )
+	{
+		bool bStillBlending = true;
+
+		float DesiredBlendSpeed = Speed * DeltaTime;
+		float BlendDif = OutCurrentValue - DesiredValue;
+
+		if (bPeak)
+		{
+			if (FMath::Abs(BlendDif) > PeakDistance + (float)FLOAT_NORMAL_THRESH)
+			{
+				SlideInCurve.JumpToEnd();
+				SlideInCurve.PlayReverse();
+			}
+			float Lerp = FMath::Max(SlideInCurve.GetLerp(), 0.1f);
+			DesiredBlendSpeed = Speed * Lerp * DeltaTime;
+		}
+
+		if (FMath::Abs(DesiredBlendSpeed) > FMath::Abs(BlendDif))
+		{
+			OutCurrentValue = DesiredValue;
+			bStillBlending = false;
+		}
+		else
+		{
+			// Reverse Direction if sliding left
+			if ( BlendDif > 0.f )
+			{
+				DesiredBlendSpeed *= -1.f;
+			}
+
+			OutCurrentValue += DesiredBlendSpeed;
+		}
+
+		return bStillBlending;
+	}
+
+	/**
+	 * Set the content of the widget to display.
+	 * @param InContent - The widget to display.
+	 */
+	void SetWidgetContent( TSharedRef<SWidget> InContent )
+	{
+		FXWidget = SNew( SFxWidget )
+		.IgnoreClipping(false)
+		.VisualOffset( this, &FCarouselDisplayItem::GetSlide )
+		[
+//			SNew(SOverlay)
+//			+ SOverlay::Slot()
+//			[
+				InContent
+//			]
+//			+ SOverlay::Slot()
+//			.HAlign(HAlign_Right)
+//			[
+//				SNew(STextBlock)
+//				.Text(this, &FCarouselDisplayItem::GetPositionText)
+//			]
+//			+ SOverlay::Slot()
+//			.HAlign(HAlign_Left)
+//			[
+//				SNew(STextBlock)
+//				.Text(this, &FCarouselDisplayItem::GetPositionText)
+//			]
+		];
+	}
+
+//	FText GetPositionText() const
+//	{
+//		return FText::FromString(FString::Printf(TEXT("%.2f, %.2f, %.2f"), SlideValue, DesiredSlideValue, SlideInCurve.GetLerp()));
+//	}
+
+	/**
+	 * Get the FX widget to display.
+	 * @return The FX widget.
+	 */
+	TSharedPtr< SFxWidget > GetSFXWidget()
+	{
+		return FXWidget;
+	}
+
+	/**
+	 * Get the widget opacity.
+	 * @param InOpacity - the desired opacity.
+	 */
+	void SetOpacity( float InOpacity )
+	{
+		OpacityValue = InOpacity;
+		FXWidget->SetColorAndOpacity( FLinearColor(1.f,1.f,1.f, OpacityValue ) );
+	}
+
+	/**
+	 * Set the slide position.
+	 * @param InSlide - the desired slide position.
+	 * @param bClearTransition - Clear any transitions.
+	 */
+	void SetSlide( float InSlide, bool bClearTransition = false)
+	{
+		SlideValue = InSlide;
+
+		if (bClearTransition)
+		{
+			bPeak = false;
+			bTransition = false;
+			DesiredSlideValue = SlideValue;
+		}
+	}
+
+	/**
+	 * Set the desired move speed.
+	 * @param DesiredMoveSpeed - the desired move speed.
+	 */
+	void SetMoveSpeed( float DesiredMoveSpeed )
+	{
+		MoveSpeed = DesiredMoveSpeed;
+	}
+
+	/**
+	 * Set the desired left limit. The widget will go to this position when scrolling out and left.
+	 * @param DesiredLeftLimit - the desired left limit.
+	 */
+	void SetSliderLeftLimit( float DesiredLeftLimit )
+	{
+		SlideValueLeftLimit = DesiredLeftLimit;
+	}
+
+	/**
+	 * Set the desired right limit. The widget will go to this position when scrolling out and right.
+	 * @param DesiredRightLimit - the desired right limit.
+	 */
+	void SetSliderRightLimit( float DesiredRightLimit )
+	{
+		SlideValueRightLimit = DesiredRightLimit;
+	}
+
+	void SetSliderPeakLeftLimit(float DesiredLimit)
+	{
+		PeakValueLeftLimit = DesiredLimit;
+	}
+
+	void SetSliderPeakRightLimit(float DesiredLimit)
+	{
+		PeakValueRightLimit = DesiredLimit;
+	}
+
+	/**
+	 * Set the optimal slide position. The widget will go to this position when in view.
+	 * @param DesiredOptimalSlideValue - the desired optimal position.
+	 */
+	void SetSliderOptimalPostion( float DesiredOptimalSlideValue )
+	{
+		OptimalSlideValue = DesiredOptimalSlideValue;
+	}
+
+	/**
+	 * Set the fade rate. The widget fade in and out at this amount per frame.
+	 * @param DesiredFadeRate - the desired fade value.
+	 */
+	void SetFadeRate( float DesiredFadeRate )
+	{
+		FadeRate = DesiredFadeRate;
+	}
+
+	bool IsInTransition()
+	{
+		return bTransition || bFade;
+	}
+
+	const float GetSlideValue() const
+	{
+		return SlideValue;
+	}
+
+protected:
+
+	/**
+	 * Get the slide position - used by the FX widget to get its position.
+	 * @return The slide position.
+	 */
+	FVector2D GetSlide() const
+	{
+		return FVector2D( SlideValue, 0.f );
+	}
+
+protected:
+	// Holds the Opacity Value. 0 == transparent
+	float OpacityValue;
+
+	// The slide position
+	float SlideValue;
+
+	// The desired opacity value
+	float DesiredOpacityValue;
+
+	// The max left position
+	float SlideValueLeftLimit;
+
+	// The max right position
+	float SlideValueRightLimit;
+
+	// The left peak position
+	float PeakValueLeftLimit;
+
+	// The right peak position
+	float PeakValueRightLimit;
+
+	// The slide value to display the widget on screen
+	float OptimalSlideValue;
+
+	// The desired slide value
+	float DesiredSlideValue;
+
+	// Holds if widget is in transition
+	bool bTransition;
+
+	// Hold if the widget is peaking
+	bool bPeak;
+
+	// Holds if the widget is fading in or out
+	bool bFade;
+
+	// Holds if the widget should fade in or out
+	bool bFadeInAndOut;
+
+	// The maximum move speed of the widget
+	float MoveSpeed;
+
+	// The Peak speed
+	float PeakSpeed;
+
+	// The Peak Distance
+	float PeakDistance;
+
+	// The amount by which a widget fades per frame
+	float FadeRate;
+
+	// Holds the FX widget that does the transitioning
+	TSharedPtr< SFxWidget > FXWidget;
+
+	// Holds the peak curve
+	FCurveSequence SlideInCurve;
+};
+
+/**
+ * Implements a widget Carousel.
+ *
+ * A widget Carousel displays widgets that can be scrolled in and out.
+ */
+template <typename ItemType>
+class SWidgetCarousel
+	: public SCompoundWidget
+{
+public:
+
+	typedef typename TSlateDelegates< ItemType >::FOnGenerateWidget FOnGenerateWidget;
+
+	SLATE_BEGIN_ARGS(SWidgetCarousel<ItemType>)
+		: _WidgetItemsSource( static_cast<const TArray<ItemType>*>(NULL) )
+		, _MoveSpeed( 1.f )
+		, _SlideValueLeftLimit( -1.f )
+		, _SlideValueRightLimit( 1.f )
+		, _PeakValueLeftLimit(-0.05f)
+		, _PeakValueRightLimit(0.05f)
+		, _FadeRate( 1.0f )
+	{ }
+
+	/** Called when we change a widget */
+	SLATE_EVENT( FOnGenerateWidget, OnGenerateWidget )
+	/** The data source */
+	SLATE_ARGUMENT( const TArray<ItemType>*, WidgetItemsSource )
+	/** The move speed */
+	SLATE_ATTRIBUTE(float, MoveSpeed )
+	/** The left limit */
+	SLATE_ATTRIBUTE(float, SlideValueLeftLimit )
+	/** The right limit */
+	SLATE_ATTRIBUTE(float, SlideValueRightLimit )
+	/** The left limit */
+	SLATE_ATTRIBUTE(float, PeakValueLeftLimit)
+	/** The right limit */
+	SLATE_ATTRIBUTE(float, PeakValueRightLimit)
+	/** The fade rate */
+	SLATE_ATTRIBUTE(float, FadeRate )
+
+	SLATE_END_ARGS()
+
+public:
+	/**
+	 * Constructs the widget.
+	 */
+	void Construct( const FArguments& InArgs )
+	{
+		OnGenerateWidget = InArgs._OnGenerateWidget;
+		ItemsSource = InArgs._WidgetItemsSource;
+		WidgetIndex = 0;
+		MoveSpeed = InArgs._MoveSpeed.Get();
+		SlideValueLeftLimit = InArgs._SlideValueLeftLimit.Get();
+		SlideValueRightLimit = InArgs._SlideValueRightLimit.Get(); 
+		PeakValueLeftLimit = InArgs._PeakValueLeftLimit.Get();
+		PeakValueRightLimit = InArgs._PeakValueRightLimit.Get();
+
+		FadeRate = InArgs._FadeRate.Get();
+
+		LeftCarouselWidget = MakeShareable(new FCarouselDisplayItem());
+		CenterCarouselWidget = MakeShareable(new FCarouselDisplayItem());
+		RightCarouselWidget = MakeShareable(new FCarouselDisplayItem());
+
+		const TArray<ItemType>& ItemsSourceRef = (*ItemsSource);
+		if (ItemsSourceRef.Num() > 0)
+		{
+			LeftCarouselWidget->SetWidgetContent(OnGenerateWidget.Execute(ItemsSourceRef[GetLeftWidgetIndex(WidgetIndex)]));
+			CenterCarouselWidget->SetWidgetContent(OnGenerateWidget.Execute(ItemsSourceRef[WidgetIndex]));
+			RightCarouselWidget->SetWidgetContent(OnGenerateWidget.Execute(ItemsSourceRef[GetRightWidgetIndex(WidgetIndex)]));
+		}
+
+		// Set up the Carousel Widget
+		CenterCarouselWidget->SetOpacity(1.0f);
+		CenterCarouselWidget->SetSlide(0.f);
+		CenterCarouselWidget->SetMoveSpeed(MoveSpeed);
+		CenterCarouselWidget->SetFadeRate(FadeRate);
+
+		LeftCarouselWidget->SetOpacity(1.0f);
+		LeftCarouselWidget->SetSlide(SlideValueLeftLimit);
+		LeftCarouselWidget->SetMoveSpeed(MoveSpeed);
+		LeftCarouselWidget->SetFadeRate(FadeRate);
+
+		RightCarouselWidget->SetOpacity(1.0f);
+		RightCarouselWidget->SetSlide(SlideValueRightLimit);
+		RightCarouselWidget->SetMoveSpeed(MoveSpeed);
+		RightCarouselWidget->SetFadeRate(FadeRate);
+
+		SetSliderLimits();
+
+		// Create the widget
+		ChildSlot
+		[
+			SAssignNew(WidgetDisplayBox, SHorizontalBox)
+		];
+
+		GenerateWidgets();
+	}
+
+	/**
+	 * Get the current widget index.
+	 * @return The visible widget index
+	 */
+	int32 GetWidgetIndex()
+	{
+		return WidgetIndex;
+	}
+
+	/**
+	 * Sets the active widget to display at the specified index.
+	 *
+	 * @param Index - The desired widget index.
+	 */
+	void SetActiveWidgetIndex(int32 Index)
+	{
+		if (!IsInTransition())
+		{
+			// Don't change if the widget index is the same
+			if (Index != WidgetIndex)
+			{
+				// Choose the scroll direction
+				EWidgetCarouselScrollDirection::Type ScrollDirection = Index > WidgetIndex ?
+					EWidgetCarouselScrollDirection::Carousel_Right : EWidgetCarouselScrollDirection::Carousel_Left;
+
+				SwapBuffer(ScrollDirection, Index);
+			}
+		}
+	}
+
+	/**
+	 * Scroll right.
+	 */
+	void SetNextWidget()
+	{
+		if (!IsInTransition())
+		{
+			SwapBuffer(EWidgetCarouselScrollDirection::Carousel_Left);
+		}
+	}
+
+	/**
+	 * Scroll left.
+	 */
+	void SetPreviousWidget()
+	{
+		if (!IsInTransition())
+		{
+			SwapBuffer(EWidgetCarouselScrollDirection::Carousel_Right);
+		}
+	}
+
+	void Peak(EWidgetCarouselScrollDirection::Type Direction)
+	{
+		if (!IsInTransition())
+		{
+			CenterCarouselWidget->PeakIn(Direction);
+			RightCarouselWidget->PeakIn(Direction);
+			LeftCarouselWidget->PeakIn(Direction);
+		}
+	}
+
+	/**
+	 * Set the item source, and scroll to the first item.
+	 * @param The item source
+	 */
+	void SetItemSource( const TArray<ItemType>* InWidgetItemsSource )
+	{
+		ItemsSource = InWidgetItemsSource;
+		if ( ItemsSource->Num() )
+		{
+			SetActiveWidgetIndex( 0 );
+		}
+	}
+
+	float GetPrimarySlide()
+	{
+		return CenterCarouselWidget->GetSlideValue();
+	}
+
+protected:
+
+	/**
+	 * Tick function to tick the carousel widgets.
+	 * @param AllottedGeometry - Widget default property - just pass to the super
+	 * @param InCurrentTime - current time
+	 * @param InDeltaTime - the frame delta time
+	 */
+	void Tick( const FGeometry& AllottedGeometry, const double InCurrentTime, const float InDeltaTime ) override
+	{
+		SCompoundWidget::Tick( AllottedGeometry, InCurrentTime, InDeltaTime );
+
+		// Tick the widgets to animate them
+		LeftCarouselWidget->Tick(InDeltaTime);
+		CenterCarouselWidget->Tick(InDeltaTime);
+		RightCarouselWidget->Tick(InDeltaTime);
+	}
+
+	/**
+	 * Regenerate the widgets to display them.
+	 */
+	void GenerateWidgets( )
+	{
+		WidgetDisplayBox->ClearChildren();
+		WidgetDisplayBox->AddSlot()
+		[
+			SNew( SOverlay )
+			+SOverlay::Slot()
+			[
+				LeftCarouselWidget->GetSFXWidget().ToSharedRef()
+			]
+			+SOverlay::Slot()
+			[
+				RightCarouselWidget->GetSFXWidget().ToSharedRef()
+			]
+			+ SOverlay::Slot()
+			[
+				CenterCarouselWidget->GetSFXWidget().ToSharedRef()
+			]
+		];
+	}
+
+	/**
+	 * Swap the buffer index.
+	 */
+	void SwapBuffer(EWidgetCarouselScrollDirection::Type ScrollDirection, int32 OverrideWidget = INDEX_NONE)
+	{
+		const TArray<ItemType>& ItemsSourceRef = (*ItemsSource);
+
+		if (ItemsSourceRef.Num() > 0)
+		{
+			if (ScrollDirection == EWidgetCarouselScrollDirection::Carousel_Left)
+			{
+				WidgetIndex = GetLeftWidgetIndex(WidgetIndex);
+				TSharedPtr<FCarouselDisplayItem> TempWidget = RightCarouselWidget;
+				RightCarouselWidget = CenterCarouselWidget;
+				CenterCarouselWidget = LeftCarouselWidget;
+				LeftCarouselWidget = TempWidget;
+				LeftCarouselWidget->SetWidgetContent(OnGenerateWidget.Execute(ItemsSourceRef[GetLeftWidgetIndex(WidgetIndex)]));
+				LeftCarouselWidget->SetSlide(CenterCarouselWidget->GetSlideValue() + SlideValueLeftLimit, true);
+			}
+			else
+			{
+				WidgetIndex = GetRightWidgetIndex(WidgetIndex);
+				TSharedPtr<FCarouselDisplayItem> TempWidget = LeftCarouselWidget;
+				LeftCarouselWidget = CenterCarouselWidget;
+				CenterCarouselWidget = RightCarouselWidget;
+				RightCarouselWidget = TempWidget;
+				RightCarouselWidget->SetWidgetContent(OnGenerateWidget.Execute(ItemsSourceRef[GetRightWidgetIndex(WidgetIndex)]));
+				RightCarouselWidget->SetSlide(CenterCarouselWidget->GetSlideValue() + SlideValueRightLimit, true);
+			}
+		}
+
+		if (OverrideWidget != INDEX_NONE)
+		{
+			// Need to find a way to clear left and right images
+			WidgetIndex = OverrideWidget;
+			CenterCarouselWidget->SetWidgetContent(OnGenerateWidget.Execute(ItemsSourceRef[OverrideWidget]));
+		}
+
+		SetSliderLimits();
+		GenerateWidgets();
+		CenterCarouselWidget->ScrollIn(ScrollDirection);
+		RightCarouselWidget->ScrollIn(ScrollDirection);
+		LeftCarouselWidget->ScrollIn(ScrollDirection);
+	}
+
+	void SetSliderLimits()
+	{
+		LeftCarouselWidget->SetSliderLeftLimit(SlideValueLeftLimit);
+		LeftCarouselWidget->SetSliderRightLimit(SlideValueRightLimit);
+		LeftCarouselWidget->SetSliderPeakLeftLimit(SlideValueLeftLimit + PeakValueLeftLimit);
+		LeftCarouselWidget->SetSliderPeakRightLimit(SlideValueLeftLimit + PeakValueRightLimit);
+		LeftCarouselWidget->SetSliderOptimalPostion(SlideValueLeftLimit);
+
+		RightCarouselWidget->SetSliderLeftLimit(SlideValueLeftLimit);
+		RightCarouselWidget->SetSliderRightLimit(SlideValueRightLimit);
+		RightCarouselWidget->SetSliderPeakLeftLimit(SlideValueRightLimit + PeakValueLeftLimit);
+		RightCarouselWidget->SetSliderPeakRightLimit(SlideValueRightLimit + PeakValueRightLimit);
+		RightCarouselWidget->SetSliderOptimalPostion(SlideValueRightLimit);
+
+		CenterCarouselWidget->SetSliderLeftLimit(SlideValueLeftLimit);
+		CenterCarouselWidget->SetSliderRightLimit(SlideValueRightLimit);
+		CenterCarouselWidget->SetSliderPeakLeftLimit(PeakValueLeftLimit);
+		CenterCarouselWidget->SetSliderPeakRightLimit(PeakValueRightLimit);
+		CenterCarouselWidget->SetSliderOptimalPostion(0);
+	}
+
+	int32 GetLeftWidgetIndex(int32 Index)
+	{
+		int32 LeftWidgetIndex = Index - 1;
+		if (LeftWidgetIndex < 0)
+		{
+			LeftWidgetIndex = FMath::Max(0, ItemsSource->Num() - 1);
+		}
+		return LeftWidgetIndex;
+	}
+
+	int32 GetRightWidgetIndex(int32 Index)
+	{
+		int32 RightWidgetIndex = Index + 1;
+		if (RightWidgetIndex > ItemsSource->Num() - 1)
+		{
+			RightWidgetIndex = 0;
+		}
+		return RightWidgetIndex;
+	}
+
+	bool IsInTransition()
+	{
+		return CenterCarouselWidget->IsInTransition() || LeftCarouselWidget->IsInTransition() || RightCarouselWidget->IsInTransition();
+	}
+
+protected:
+
+	// Holds the CarouselWidgets
+	TSharedPtr< FCarouselDisplayItem > LeftCarouselWidget;
+	TSharedPtr< FCarouselDisplayItem > CenterCarouselWidget;
+	TSharedPtr< FCarouselDisplayItem > RightCarouselWidget;
+
+	// Pointer to the array of data items that we are observing
+	const TArray<ItemType>* ItemsSource;
+
+	// Hold the delegate to be invoked when a widget changes.
+	FOnGenerateWidget OnGenerateWidget;
+
+	// Holds the widget display box.
+	TSharedPtr< SHorizontalBox > WidgetDisplayBox;
+
+	// Holds the visible widget index.
+	int32 WidgetIndex;
+
+	// Holds the move speed
+	float MoveSpeed;
+
+	// Holds The left limit
+	float SlideValueLeftLimit;
+
+	// The right limit
+	float SlideValueRightLimit;
+
+	// Holds The peak left limit
+	float PeakValueLeftLimit;
+
+	// The peak right limit
+	float PeakValueRightLimit;
+
+	// The fade rate
+	float FadeRate;
+};
diff --git a/Engine/Source/Developer/WidgetCarousel/Public/SWidgetCarouselWithNavigation.h b/Engine/Source/Developer/WidgetCarousel/Public/SWidgetCarouselWithNavigation.h
new file mode 100644
--- /dev/null
+++ b/Engine/Source/Developer/WidgetCarousel/Public/SWidgetCarouselWithNavigation.h
@@ -0,0 +1,169 @@
+// Copyright 1998-2014 Epic Games, Inc. All Rights Reserved.
+
+#pragma once
+
+#include "SWidgetCarousel.h"
+#include "SCarouselNavigationButton.h"
+#include "SCarouselNavigationBar.h"
+
+#include "SScaleBox.h"
+
+/** A widget carousel which includes embedded left and right navigation buttons and a navigation bar below. */
+template <typename ItemType>
+class SWidgetCarouselWithNavigation
+	: public SCompoundWidget
+{
+public:
+	typedef typename TSlateDelegates< ItemType >::FOnGenerateWidget FOnGenerateWidget;
+
+	SLATE_BEGIN_ARGS(SWidgetCarouselWithNavigation<ItemType>)
+		: _WidgetItemsSource(static_cast<const TArray<ItemType>*>(NULL))
+	{ }
+
+		/** Called when we change a widget */
+		SLATE_EVENT(FOnGenerateWidget, OnGenerateWidget)
+
+		/** The data source */
+		SLATE_ARGUMENT(const TArray<ItemType>*, WidgetItemsSource)
+
+	SLATE_END_ARGS();
+
+	void Construct(const FArguments& InArgs)
+	{
+		OnGenerateWidget = InArgs._OnGenerateWidget;
+		WidgetItemsSource = InArgs._WidgetItemsSource;
+
+		ChildSlot
+		[
+			SNew(SVerticalBox)
+			+ SVerticalBox::Slot()
+			.AutoHeight()
+			.HAlign(HAlign_Center)
+			.VAlign(VAlign_Center)
+			[
+				SNew(SBox)
+				.MaxDesiredHeight(442)
+				[
+					SNew(SOverlay)
+					+ SOverlay::Slot()
+					.HAlign(HAlign_Center)
+					.VAlign(VAlign_Center)
+					[
+						SNew(SBorder)
+						.BorderImage(FWidgetCarouselStyle::Get().GetBrush("Black .6"))
+						.Padding(0)
+						[
+							SNew(SScaleBox)
+							.Stretch(EStretch::ScaleToFit)
+							.StretchDirection(EStretchDirection::DownOnly)
+							[
+								SAssignNew(Carousel, SWidgetCarousel<ItemType>)
+								.WidgetItemsSource(WidgetItemsSource)
+								.FadeRate(0)
+								.SlideValueLeftLimit(-1)
+								.SlideValueRightLimit(1)
+								.MoveSpeed(5.f)
+								.OnGenerateWidget(this, &SWidgetCarouselWithNavigation::GenerateWidget)
+							]
+						]
+					]
+
+					+ SOverlay::Slot()
+					.Padding(0)
+					[
+						SNew(SHorizontalBox)
+						+ SHorizontalBox::Slot()
+						.HAlign(HAlign_Left)
+						[
+							SNew(SCarouselNavigationButton)
+							.OnClicked(this, &SWidgetCarouselWithNavigation::HandleNextButtonClicked)
+							.Visibility(this, &SWidgetCarouselWithNavigation::GetScreenshotNavigationVisibility)
+							.OnBeginPeek(this, &SWidgetCarouselWithNavigation::HandlePeak, EWidgetCarouselScrollDirection::Carousel_Right)
+							.OnEndPeek(this, &SWidgetCarouselWithNavigation::HandlePeak, EWidgetCarouselScrollDirection::Carousel_Center)
+							.Direction(SCarouselNavigationButton::ENavigationButtonDirection::Left)
+						]
+						+ SHorizontalBox::Slot()
+						.HAlign(HAlign_Right)
+						[
+							SNew(SCarouselNavigationButton)
+							.OnClicked(this, &SWidgetCarouselWithNavigation::HandlePreviousButtonClicked)
+							.Visibility(this, &SWidgetCarouselWithNavigation::GetScreenshotNavigationVisibility)
+							.OnBeginPeek(this, &SWidgetCarouselWithNavigation::HandlePeak, EWidgetCarouselScrollDirection::Carousel_Left)
+							.OnEndPeek(this, &SWidgetCarouselWithNavigation::HandlePeak, EWidgetCarouselScrollDirection::Carousel_Center)
+							.Direction(SCarouselNavigationButton::ENavigationButtonDirection::Right)
+						]
+					]
+				]
+			]
+			+ SVerticalBox::Slot()
+			.AutoHeight()
+			.HAlign(HAlign_Center)
+			.Padding(10.0f, 10.f)
+			[
+				SNew(SCarouselNavigationBar)
+				.Visibility(this, &SWidgetCarouselWithNavigation::GetScreenshotNavigationVisibility)
+				.ItemCount(WidgetItemsSource->Num())
+				.CurrentItemIndex(this, &SWidgetCarouselWithNavigation::GetCurrentItemIndex)
+				.CurrentSlideAmount(this, &SWidgetCarouselWithNavigation::GetCurrentSlideAmount)
+				.OnSelectedIndexChanged(this, &SWidgetCarouselWithNavigation::CarouselScrollBarIndexChanged)
+			]
+		];
+	}
+
+	~SWidgetCarouselWithNavigation()
+	{
+	}
+
+private:
+	TSharedRef<SWidget> GenerateWidget(ItemType Item)
+	{
+		return OnGenerateWidget.Execute(Item);
+	}
+
+	FReply HandleNextButtonClicked()
+	{
+		Carousel->SetNextWidget();
+		return FReply::Handled();
+	}
+
+	FReply HandlePreviousButtonClicked()
+	{
+		Carousel->SetPreviousWidget();
+		return FReply::Handled();
+	}
+
+	void HandlePeak(EWidgetCarouselScrollDirection::Type Direction)
+	{
+		Carousel->Peak(Direction);
+	}
+
+	EVisibility GetScreenshotVisibility() const
+	{
+		return WidgetItemsSource->Num() ? EVisibility::Visible : EVisibility::Collapsed;
+	}
+
+	EVisibility GetScreenshotNavigationVisibility() const
+	{
+		return WidgetItemsSource->Num() > 1 ? EVisibility::Visible : EVisibility::Collapsed;
+	}
+
+	int32 GetCurrentItemIndex() const
+	{
+		return Carousel->GetWidgetIndex();
+	}
+
+	float GetCurrentSlideAmount() const
+	{
+		return Carousel->GetPrimarySlide();
+	}
+
+	void CarouselScrollBarIndexChanged(int32 NewIndex)
+	{
+		Carousel->SetActiveWidgetIndex(NewIndex);
+	}
+
+private:
+	TSharedPtr<SWidgetCarousel<ItemType>> Carousel;
+	FOnGenerateWidget OnGenerateWidget;
+	const TArray<ItemType>* WidgetItemsSource;
+};
\ No newline at end of file
diff --git a/Engine/Source/Developer/WidgetCarousel/Public/WidgetCarouselModule.h b/Engine/Source/Developer/WidgetCarousel/Public/WidgetCarouselModule.h
new file mode 100644
--- /dev/null
+++ b/Engine/Source/Developer/WidgetCarousel/Public/WidgetCarouselModule.h
@@ -0,0 +1,15 @@
+// Copyright 1998-2014 Epic Games, Inc. All Rights Reserved.
+#pragma once
+
+/**
+ * Implements the WidgetCarousel module.
+ */
+class WIDGETCAROUSEL_API FWidgetCarouselModule
+	: public IModuleInterface
+{
+public:
+
+	virtual void StartupModule( ) override;
+
+	virtual void ShutdownModule( ) override;
+};
diff --git a/Engine/Source/Developer/WidgetCarousel/Public/WidgetCarouselStyle.h b/Engine/Source/Developer/WidgetCarousel/Public/WidgetCarouselStyle.h
new file mode 100644
--- /dev/null
+++ b/Engine/Source/Developer/WidgetCarousel/Public/WidgetCarouselStyle.h
@@ -0,0 +1,31 @@
+// Copyright 1998-2014 Epic Games, Inc. All Rights Reserved.
+
+#pragma once
+
+#include "SlateBasics.h"
+
+/** Manages the style which provides resources for the add content dialog. */
+class WIDGETCAROUSEL_API FWidgetCarouselStyle
+{
+public:
+
+	static void Initialize();
+
+	static void Shutdown();
+
+	/** reloads textures used by slate renderer */
+	static void ReloadTextures();
+
+	/** @return The Slate style set for the Shooter game */
+	static const ISlateStyle& Get();
+
+	static FName GetStyleSetName();
+
+private:
+
+	static TSharedRef< class FSlateStyleSet > Create();
+
+private:
+
+	static TSharedPtr< class FSlateStyleSet > WidgetCarouselStyleInstance;
+};
\ No newline at end of file
diff --git a/Engine/Source/Developer/WidgetCarousel/WidgetCarousel.Build.cs b/Engine/Source/Developer/WidgetCarousel/WidgetCarousel.Build.cs
new file mode 100644
--- /dev/null
+++ b/Engine/Source/Developer/WidgetCarousel/WidgetCarousel.Build.cs
@@ -0,0 +1,25 @@
+// Copyright 1998-2014 Epic Games, Inc. All Rights Reserved.
+
+using UnrealBuildTool;
+
+public class WidgetCarousel : ModuleRules
+{
+	public WidgetCarousel(TargetInfo Target)
+	{
+		PrivateDependencyModuleNames.AddRange(
+			new string[] {
+				"Core",
+				"Slate",
+				"SlateCore",
+				"InputCore",
+				"EditorStyle",
+			}
+		);
+
+		PrivateIncludePaths.AddRange(
+			new string[] {
+				"Developer/WidgetCarousel/Private",
+			}
+		);
+	}
+}
