WidgetCarousel - Add missing files from initial checkin.

diff --git a/Engine/Source/Developer/WidgetCarousel/Private/SCarouselNavigationBar.cpp b/Engine/Source/Developer/WidgetCarousel/Private/SCarouselNavigationBar.cpp
new file mode 100644
--- /dev/null
+++ b/Engine/Source/Developer/WidgetCarousel/Private/SCarouselNavigationBar.cpp
@@ -0,0 +1,102 @@
+// Copyright 1998-2014 Epic Games, Inc. All Rights Reserved.
+
+#include "WidgetCarouselPrivatePCH.h"
+
+#define LOCTEXT_NAMESPACE "WidgetCarousel"
+
+int32 SCarouselNavigationBar::GetItemCount()
+{
+	return ItemCount;
+}
+
+void SCarouselNavigationBar::SetItemCount(int32 InItemCount)
+{
+	ItemCount = InItemCount;
+	BuildScrollBar();
+}
+
+void SCarouselNavigationBar::Construct(const FArguments& InArgs)
+{
+	OnSelectedIndexChanged = InArgs._OnSelectedIndexChanged;
+	ItemCount = InArgs._ItemCount;
+	CurrentItemIndex = InArgs._CurrentItemIndex;
+	CurrentSlideAmount = InArgs._CurrentSlideAmount;
+
+	ChildSlot
+	[
+		SNew(SOverlay)
+		+ SOverlay::Slot()
+		[
+			SNew(SFxWidget)
+			.IgnoreClipping(false)
+			.VisualOffset(this, &SCarouselNavigationBar::GetPlaceHolderPosition)
+			[
+				SNew(SHorizontalBox)
+				+ SHorizontalBox::Slot()
+				.AutoWidth()
+				.VAlign(VAlign_Center)
+				.HAlign(HAlign_Left)
+				[
+					SNew(SImage)
+					.Image(FWidgetCarouselStyle::Get().GetBrush(TEXT("WidgetCarousel.NavigationScrollBarHighlight")))
+				]
+			]
+		]
+		+ SOverlay::Slot()
+		[
+			SAssignNew(WidgetScrollBox, SHorizontalBox)
+		]
+	];
+
+	BuildScrollBar();
+}
+
+void SCarouselNavigationBar::BuildScrollBar()
+{
+	// May change this to have the scroll bar get their images via delegate rather than rebuild
+	WidgetScrollBox->ClearChildren();
+
+	if (ItemCount > 1)
+	{
+		for (int32 Index = 0; Index < ItemCount; Index++)
+		{
+			FString BrushName = TEXT("WidgetCarousel.NavigationScrollBarCenterButton");
+			if (Index == 0)
+			{
+				BrushName = TEXT("WidgetCarousel.NavigationScrollBarLeftButton");
+			}
+			else if (Index == ItemCount - 1)
+			{
+				BrushName = TEXT("WidgetCarousel.NavigationScrollBarRightButton");
+			}
+
+			WidgetScrollBox->AddSlot()
+			.Padding(0)
+			.AutoWidth()
+			.VAlign(VAlign_Center)
+			[
+				SNew(SButton)
+				.ButtonStyle(FWidgetCarouselStyle::Get(), *BrushName)
+				.OnClicked(this, &SCarouselNavigationBar::HandleItemButtonClicked, Index)
+				.Cursor(EMouseCursor::Hand)
+			];
+		}
+	}
+}
+
+FReply SCarouselNavigationBar::HandleItemButtonClicked(int32 ItemIndex)
+{
+	OnSelectedIndexChanged.ExecuteIfBound(ItemIndex);
+	BuildScrollBar();
+	return FReply::Handled();
+}
+
+FVector2D SCarouselNavigationBar::GetPlaceHolderPosition() const
+{
+	float WidgetSize = 1.f / (float)ItemCount;
+
+	float NewPos = CurrentItemIndex.Get() * WidgetSize;
+	NewPos -= CurrentSlideAmount.Get() * WidgetSize;
+
+	return FVector2D(NewPos, 0.f);
+}
\ No newline at end of file
diff --git a/Engine/Source/Developer/WidgetCarousel/Private/WidgetCarouselModule.cpp b/Engine/Source/Developer/WidgetCarousel/Private/WidgetCarouselModule.cpp
new file mode 100644
--- /dev/null
+++ b/Engine/Source/Developer/WidgetCarousel/Private/WidgetCarouselModule.cpp
@@ -0,0 +1,15 @@
+// Copyright 1998-2014 Epic Games, Inc. All Rights Reserved.
+
+#include "WidgetCarouselPrivatePCH.h"
+
+void FWidgetCarouselModule::StartupModule()
+{
+	FWidgetCarouselStyle::Initialize();
+}
+
+void FWidgetCarouselModule::ShutdownModule()
+{
+	FWidgetCarouselStyle::Shutdown();
+}
+
+IMPLEMENT_MODULE(FWidgetCarouselModule, WidgetCarousel);
