WidgetCarousel - Add style structs for the navigation button and bar, and add separate styles for the launcher.

diff --git a/Engine/Source/Developer/WidgetCarousel/Private/SCarouselNavigationBar.cpp b/Engine/Source/Developer/WidgetCarousel/Private/SCarouselNavigationBar.cpp
--- a/Engine/Source/Developer/WidgetCarousel/Private/SCarouselNavigationBar.cpp
+++ b/Engine/Source/Developer/WidgetCarousel/Private/SCarouselNavigationBar.cpp
@@ -17,6 +17,7 @@ void SCarouselNavigationBar::SetItemCount(int32 InItemCount)
 
 void SCarouselNavigationBar::Construct(const FArguments& InArgs)
 {
+	Style = InArgs._Style;
 	OnSelectedIndexChanged = InArgs._OnSelectedIndexChanged;
 	ItemCount = InArgs._ItemCount;
 	CurrentItemIndex = InArgs._CurrentItemIndex;
@@ -38,7 +39,7 @@ void SCarouselNavigationBar::Construct(const FArguments& InArgs)
 				.HAlign(HAlign_Left)
 				[
 					SNew(SImage)
-					.Image(FWidgetCarouselStyle::Get().GetBrush(TEXT("WidgetCarousel.NavigationScrollBarHighlight")))
+					.Image(&Style->HighlightBrush)
 				]
 			]
 		]
@@ -60,14 +61,18 @@ void SCarouselNavigationBar::BuildScrollBar()
 	{
 		for (int32 Index = 0; Index < ItemCount; Index++)
 		{
-			FString BrushName = TEXT("WidgetCarousel.NavigationScrollBarCenterButton");
+			const FButtonStyle* ButtonStyle;
 			if (Index == 0)
 			{
-				BrushName = TEXT("WidgetCarousel.NavigationScrollBarLeftButton");
+				ButtonStyle = &Style->LeftButtonStyle;
 			}
 			else if (Index == ItemCount - 1)
 			{
-				BrushName = TEXT("WidgetCarousel.NavigationScrollBarRightButton");
+				ButtonStyle = &Style->RightButtonStyle;
+			}
+			else
+			{
+				ButtonStyle = &Style->CenterButtonStyle;
 			}
 
 			WidgetScrollBox->AddSlot()
@@ -76,7 +81,7 @@ void SCarouselNavigationBar::BuildScrollBar()
 			.VAlign(VAlign_Center)
 			[
 				SNew(SButton)
-				.ButtonStyle(FWidgetCarouselStyle::Get(), *BrushName)
+				.ButtonStyle(ButtonStyle)
 				.OnClicked(this, &SCarouselNavigationBar::HandleItemButtonClicked, Index)
 				.Cursor(EMouseCursor::Hand)
 			];
diff --git a/Engine/Source/Developer/WidgetCarousel/Private/WidgetCarouselModule.cpp b/Engine/Source/Developer/WidgetCarousel/Private/WidgetCarouselModule.cpp
--- a/Engine/Source/Developer/WidgetCarousel/Private/WidgetCarouselModule.cpp
+++ b/Engine/Source/Developer/WidgetCarousel/Private/WidgetCarouselModule.cpp
@@ -4,12 +4,12 @@
 
 void FWidgetCarouselModule::StartupModule()
 {
-	FWidgetCarouselStyle::Initialize();
+	FWidgetCarouselModuleStyle::Initialize();
 }
 
 void FWidgetCarouselModule::ShutdownModule()
 {
-	FWidgetCarouselStyle::Shutdown();
+	FWidgetCarouselModuleStyle::Shutdown();
 }
 
 IMPLEMENT_MODULE(FWidgetCarouselModule, WidgetCarousel);
diff --git a/Engine/Source/Developer/WidgetCarousel/Private/WidgetCarouselStyle.cpp b/Engine/Source/Developer/WidgetCarousel/Private/WidgetCarouselStyle.cpp
--- a/Engine/Source/Developer/WidgetCarousel/Private/WidgetCarouselStyle.cpp
+++ b/Engine/Source/Developer/WidgetCarousel/Private/WidgetCarouselStyle.cpp
@@ -2,11 +2,60 @@
 
 #include "WidgetCarouselPrivatePCH.h"
 
-#include "EditorStyle.h"
+#define IMAGE_BRUSH( RelativePath, ... ) FSlateImageBrush( Style->RootToContentDir( RelativePath, TEXT(".png") ), __VA_ARGS__ )
+#define BOX_BRUSH( RelativePath, ... ) FSlateBoxBrush( Style->RootToContentDir( RelativePath, TEXT(".png") ), __VA_ARGS__ )
+
+/**
+	Navigation button style
+*/
+FWidgetCarouselNavigationButtonStyle::FWidgetCarouselNavigationButtonStyle()
+{
+}
+
+const FName FWidgetCarouselNavigationButtonStyle::TypeName(TEXT("FWidgetCarouselNavigationButtonStyle"));
+
+const FWidgetCarouselNavigationButtonStyle& FWidgetCarouselNavigationButtonStyle::GetDefault()
+{
+	static FWidgetCarouselNavigationButtonStyle Default;
+	return Default;
+}
+
+void FWidgetCarouselNavigationButtonStyle::GetResources(TArray< const FSlateBrush* > & OutBrushes) const
+{
+	OutBrushes.Add(&NavigationButtonLeftImage);
+	OutBrushes.Add(&NavigationButtonRightImage);
+	InnerButtonStyle.GetResources(OutBrushes);
+}
+
+/**
+Navigation bar style
+*/
+FWidgetCarouselNavigationBarStyle::FWidgetCarouselNavigationBarStyle()
+{
+}
+
+const FName FWidgetCarouselNavigationBarStyle::TypeName(TEXT("FWidgetCarouselNavigationBarStyle"));
+
+const FWidgetCarouselNavigationBarStyle& FWidgetCarouselNavigationBarStyle::GetDefault()
+{
+	static FWidgetCarouselNavigationBarStyle Default;
+	return Default;
+}
+
+void FWidgetCarouselNavigationBarStyle::GetResources(TArray< const FSlateBrush* > & OutBrushes) const
+{
+	OutBrushes.Add(&HighlightBrush);
+	LeftButtonStyle.GetResources(OutBrushes);
+	CenterButtonStyle.GetResources(OutBrushes);
+	RightButtonStyle.GetResources(OutBrushes);
+}
 
-TSharedPtr< FSlateStyleSet > FWidgetCarouselStyle::WidgetCarouselStyleInstance = NULL;
+/**
+	Module style set
+*/
+TSharedPtr< FSlateStyleSet > FWidgetCarouselModuleStyle::WidgetCarouselStyleInstance = NULL;
 
-void FWidgetCarouselStyle::Initialize()
+void FWidgetCarouselModuleStyle::Initialize()
 {
 	if ( !WidgetCarouselStyleInstance.IsValid() )
 	{
@@ -15,30 +64,20 @@ void FWidgetCarouselStyle::Initialize()
 	}
 }
 
-void FWidgetCarouselStyle::Shutdown()
+void FWidgetCarouselModuleStyle::Shutdown()
 {
 	FSlateStyleRegistry::UnRegisterSlateStyle( *WidgetCarouselStyleInstance );
 	ensure( WidgetCarouselStyleInstance.IsUnique() );
 	WidgetCarouselStyleInstance.Reset();
 }
 
-FName FWidgetCarouselStyle::GetStyleSetName()
+FName FWidgetCarouselModuleStyle::GetStyleSetName()
 {
-	static FName StyleSetName(TEXT("WidgetCarouselStyle"));
+	static FName StyleSetName(TEXT("WidgetCarousel"));
 	return StyleSetName;
 }
 
-#define IMAGE_BRUSH( RelativePath, ... ) FSlateImageBrush( Style->RootToContentDir( RelativePath, TEXT(".png") ), __VA_ARGS__ )
-#define BOX_BRUSH( RelativePath, ... ) FSlateBoxBrush( Style->RootToContentDir( RelativePath, TEXT(".png") ), __VA_ARGS__ )
-#define BORDER_BRUSH( RelativePath, ... ) FSlateBorderBrush( Style->RootToContentDir( RelativePath, TEXT(".png") ), __VA_ARGS__ )
-#define TTF_FONT( RelativePath, ... ) FSlateFontInfo( Style->RootToContentDir( RelativePath, TEXT(".ttf") ), __VA_ARGS__ )
-#define OTF_FONT( RelativePath, ... ) FSlateFontInfo( Style->RootToContentDir( RelativePath, TEXT(".otf") ), __VA_ARGS__ )
-
-const FVector2D Icon16x16(16.0f, 16.0f);
-const FVector2D Icon20x20(20.0f, 20.0f);
-const FVector2D Icon40x40(40.0f, 40.0f);
-
-TSharedRef< FSlateStyleSet > FWidgetCarouselStyle::Create()
+TSharedRef< FSlateStyleSet > FWidgetCarouselModuleStyle::Create()
 {
 	TSharedRef< FSlateStyleSet > Style = MakeShareable(new FSlateStyleSet("WidgetCarouselStyle"));
 	Style->SetContentRoot(FPaths::EngineContentDir() / TEXT("Editor/Slate/WidgetCarousel"));
@@ -48,54 +87,46 @@ TSharedRef< FSlateStyleSet > FWidgetCarouselStyle::Create()
 		.SetPressedPadding(0);
 
 	FLinearColor PrimaryCallToActionColor = FLinearColor(1.0, 0.7372, 0.05637);
-	FLinearColor PrimaryCallToActionColor_Hovered = FLinearColor(1.0, 0.83553, 0.28445);
-	FLinearColor PrimaryCallToActionColor_Pressed = FLinearColor(1.0, 0.66612, 0.0012);
-
-	Style->Set("PrimaryColor", PrimaryCallToActionColor);
-
-	Style->Set("WidgetCarousel.NavigationButtonLeft", new IMAGE_BRUSH("Arrow-Left", FVector2D(25.0f, 42.0f)));
-
-	Style->Set("WidgetCarousel.NavigationButtonRight", new IMAGE_BRUSH("Arrow-Right", FVector2D(25.0f, 42.0f)));
-
-	Style->Set("WidgetCarousel.NavigationButton", FButtonStyle(DefaultButton)
-		.SetNormal(BOX_BRUSH("WhiteBox_7px_CornerRadius", FVector2D(17, 17), FMargin(0.5), PrimaryCallToActionColor))
-		.SetPressed(BOX_BRUSH("WhiteBox_7px_CornerRadius", FVector2D(17, 17), FMargin(0.5), PrimaryCallToActionColor_Pressed))
-		.SetHovered(BOX_BRUSH("WhiteBox_7px_CornerRadius", FVector2D(17, 17), FMargin(0.5), PrimaryCallToActionColor_Hovered)));
-
-	Style->Set("WidgetCarousel.NavigationScrollBarHighlight", new IMAGE_BRUSH("CarouselNavMarker", FVector2D(80.0f, 20.0f), PrimaryCallToActionColor));
-
-	Style->Set("WidgetCarousel.NavigationScrollBarLeftButton", FButtonStyle(DefaultButton)
-		.SetNormal(IMAGE_BRUSH("CarouselNavLeft", FVector2D(80.0f, 20.0f), FLinearColor::White))
-		.SetHovered(IMAGE_BRUSH("CarouselNavLeft", FVector2D(80.0f, 20.0f), PrimaryCallToActionColor))
-		.SetPressed(IMAGE_BRUSH("CarouselNavLeft", FVector2D(80.0f, 20.0f), PrimaryCallToActionColor_Pressed)));
-
-	Style->Set("WidgetCarousel.NavigationScrollBarCenterButton", FButtonStyle(DefaultButton)
-		.SetNormal(IMAGE_BRUSH("CarouselNavCenter", FVector2D(80.0f, 20.0f), FLinearColor::White))
-		.SetHovered(IMAGE_BRUSH("CarouselNavCenter", FVector2D(80.0f, 20.0f), PrimaryCallToActionColor))
-		.SetPressed(IMAGE_BRUSH("CarouselNavCenter", FVector2D(80.0f, 20.0f), PrimaryCallToActionColor_Pressed)));
-
-	Style->Set("WidgetCarousel.NavigationScrollBarRightButton", FButtonStyle(DefaultButton)
-		.SetNormal(IMAGE_BRUSH("CarouselNavRight", FVector2D(80.0f, 20.0f), FLinearColor::White))
-		.SetHovered(IMAGE_BRUSH("CarouselNavRight", FVector2D(80.0f, 20.0f), PrimaryCallToActionColor))
-		.SetPressed(IMAGE_BRUSH("CarouselNavRight", FVector2D(80.0f, 20.0f), PrimaryCallToActionColor_Pressed)));
+	FLinearColor PrimaryCallToActionColorHovered = FLinearColor(1.0, 0.83553, 0.28445);
+	FLinearColor PrimaryCallToActionColorPressed = FLinearColor(1.0, 0.66612, 0.0012);
+
+	Style->Set("CarouselNavigationButton", FWidgetCarouselNavigationButtonStyle()
+		.SetInnerButtonStyle(FButtonStyle(DefaultButton)
+			.SetNormal(BOX_BRUSH("WhiteBox_7px_CornerRadius", FVector2D(17, 17), FMargin(0.5), PrimaryCallToActionColor))
+			.SetPressed(BOX_BRUSH("WhiteBox_7px_CornerRadius", FVector2D(17, 17), FMargin(0.5), PrimaryCallToActionColorPressed))
+			.SetHovered(BOX_BRUSH("WhiteBox_7px_CornerRadius", FVector2D(17, 17), FMargin(0.5), PrimaryCallToActionColorHovered)))
+		.SetNavigationButtonLeftImage(IMAGE_BRUSH("Arrow-Left", FVector2D(25.0f, 42.0f)))
+		.SetNavigationButtonRightImage(IMAGE_BRUSH("Arrow-Right", FVector2D(25.0f, 42.0f))));
+
+	Style->Set("CarouselNavigationBar", FWidgetCarouselNavigationBarStyle()
+		.SetHighlightBrush(IMAGE_BRUSH("CarouselNavMarker", FVector2D(80.0f, 20.0f), PrimaryCallToActionColor))
+		.SetLeftButtonStyle(FButtonStyle(DefaultButton)
+			.SetNormal(IMAGE_BRUSH("CarouselNavLeft", FVector2D(80.0f, 20.0f), FLinearColor::White))
+			.SetHovered(IMAGE_BRUSH("CarouselNavLeft", FVector2D(80.0f, 20.0f), PrimaryCallToActionColor))
+			.SetPressed(IMAGE_BRUSH("CarouselNavLeft", FVector2D(80.0f, 20.0f), PrimaryCallToActionColorPressed)))
+		.SetCenterButtonStyle(FButtonStyle()
+			.SetNormal(IMAGE_BRUSH("CarouselNavCenter", FVector2D(80.0f, 20.0f), FLinearColor::White))
+			.SetHovered(IMAGE_BRUSH("CarouselNavCenter", FVector2D(80.0f, 20.0f), PrimaryCallToActionColor))
+			.SetPressed(IMAGE_BRUSH("CarouselNavCenter", FVector2D(80.0f, 20.0f), PrimaryCallToActionColorPressed)))
+		.SetRightButtonStyle(FButtonStyle()
+			.SetNormal(IMAGE_BRUSH("CarouselNavRight", FVector2D(80.0f, 20.0f), FLinearColor::White))
+			.SetHovered(IMAGE_BRUSH("CarouselNavRight", FVector2D(80.0f, 20.0f), PrimaryCallToActionColor))
+			.SetPressed(IMAGE_BRUSH("CarouselNavRight", FVector2D(80.0f, 20.0f), PrimaryCallToActionColorPressed))));
 
 	Style->Set("Black .6", FLinearColor(0, 0, 0, 0.6f));
 
 	return Style;
 }
 
-#undef IMAGE_BRUSH
-#undef BOX_BRUSH
-#undef BORDER_BRUSH
-#undef TTF_FONT
-#undef OTF_FONT
-
-void FWidgetCarouselStyle::ReloadTextures()
+void FWidgetCarouselModuleStyle::ReloadTextures()
 {
 	FSlateApplication::Get().GetRenderer()->ReloadTextureResources();
 }
 
-const ISlateStyle& FWidgetCarouselStyle::Get()
+const ISlateStyle& FWidgetCarouselModuleStyle::Get()
 {
 	return *WidgetCarouselStyleInstance;
 }
+
+#undef IMAGE_BRUSH
+#undef BOX_BRUSH
diff --git a/Engine/Source/Developer/WidgetCarousel/Public/SCarouselNavigationBar.h b/Engine/Source/Developer/WidgetCarousel/Public/SCarouselNavigationBar.h
--- a/Engine/Source/Developer/WidgetCarousel/Public/SCarouselNavigationBar.h
+++ b/Engine/Source/Developer/WidgetCarousel/Public/SCarouselNavigationBar.h
@@ -12,12 +12,16 @@ class WIDGETCAROUSEL_API SCarouselNavigationBar : public SCompoundWidget
 
 public:
 	SLATE_BEGIN_ARGS(SCarouselNavigationBar)
-		: _OnSelectedIndexChanged()
+		: _Style(&FWidgetCarouselModuleStyle::Get().GetWidgetStyle<FWidgetCarouselNavigationBarStyle>("CarouselNavigationBar")) 
+		, _OnSelectedIndexChanged()
 		, _ItemCount(0)
 		, _CurrentItemIndex(0)
 		, _CurrentSlideAmount(0.0f)
+		
 	{}
 
+		SLATE_STYLE_ARGUMENT(FWidgetCarouselNavigationBarStyle, Style)
+
 		/** Called when the selected index changes. */
 		SLATE_EVENT(FOnSelectedIndexChanged, OnSelectedIndexChanged)
 
@@ -46,6 +50,7 @@ private:
 	FVector2D GetPlaceHolderPosition() const;
 
 private:
+	const FWidgetCarouselNavigationBarStyle* Style;
 	TSharedPtr<SHorizontalBox> WidgetScrollBox;
 	FOnSelectedIndexChanged OnSelectedIndexChanged;
 	int32 ItemCount;
diff --git a/Engine/Source/Developer/WidgetCarousel/Public/SCarouselNavigationButton.h b/Engine/Source/Developer/WidgetCarousel/Public/SCarouselNavigationButton.h
--- a/Engine/Source/Developer/WidgetCarousel/Public/SCarouselNavigationButton.h
+++ b/Engine/Source/Developer/WidgetCarousel/Public/SCarouselNavigationButton.h
@@ -94,7 +94,11 @@ public:
 		Right
 	};
 
-	SLATE_BEGIN_ARGS(SCarouselNavigationButton) {}
+	SLATE_BEGIN_ARGS(SCarouselNavigationButton)
+		: _Style(&FWidgetCarouselModuleStyle::Get().GetWidgetStyle<FWidgetCarouselNavigationButtonStyle>("CarouselNavigationButton"))
+		{}
+
+		SLATE_STYLE_ARGUMENT(FWidgetCarouselNavigationButtonStyle, Style)
 
 		/** Called when the button is clicked */
 		SLATE_EVENT(FOnClicked, OnClicked)
@@ -111,6 +115,7 @@ public:
 	void Construct(const FArguments& InArgs)
 	{
 		ImageTransparency = 0.f;
+		Style = InArgs._Style;
 
 		ChildSlot
 		[
@@ -144,7 +149,7 @@ public:
 							.VAlign(VAlign_Center)
 							.ToolTipText(InArgs._ToolTipText)
 							.OnClicked(InArgs._OnClicked)
-							.ButtonStyle(FWidgetCarouselStyle::Get(), "WidgetCarousel.NavigationButton")
+							.ButtonStyle(&Style->InnerButtonStyle)
 							.ContentPadding(FMargin(15, 30))
 							.ButtonColorAndOpacity(this, &SCarouselNavigationButton::GetButtonColor)
 							.Cursor(EMouseCursor::Hand)
@@ -155,8 +160,7 @@ public:
 								[
 									SNew(SImage)
 									.Image(InArgs._Direction == ENavigationButtonDirection::Left ?
-										FWidgetCarouselStyle::Get().GetBrush("WidgetCarousel.NavigationButtonLeft") :
-										FWidgetCarouselStyle::Get().GetBrush("WidgetCarousel.NavigationButtonRight"))
+										&Style->NavigationButtonLeftImage : &Style->NavigationButtonRightImage)
 									.ColorAndOpacity(this, &SCarouselNavigationButton::GetButtonImageColor)
 								]
 							]
@@ -180,7 +184,7 @@ private:
 		{
 			if (!VisibleButton->IsHovered())
 			{
-				FLinearColor Color = FWidgetCarouselStyle::Get().GetColor("PrimaryColor");
+				FLinearColor Color = FLinearColor::White;
 				Color.A = ImageTransparency;
 				return Color;
 			}
@@ -241,7 +245,7 @@ private:
 	}
 
 private:
-
+	const FWidgetCarouselNavigationButtonStyle* Style;
 	TSharedPtr<SButton> HiddenButton;
 	TSharedPtr<SButton> VisibleButton;
 
diff --git a/Engine/Source/Developer/WidgetCarousel/Public/SWidgetCarouselWithNavigation.h b/Engine/Source/Developer/WidgetCarousel/Public/SWidgetCarouselWithNavigation.h
--- a/Engine/Source/Developer/WidgetCarousel/Public/SWidgetCarouselWithNavigation.h
+++ b/Engine/Source/Developer/WidgetCarousel/Public/SWidgetCarouselWithNavigation.h
@@ -50,7 +50,7 @@ public:
 					.VAlign(VAlign_Center)
 					[
 						SNew(SBorder)
-						.BorderImage(FWidgetCarouselStyle::Get().GetBrush("Black .6"))
+						.BorderImage(FWidgetCarouselModuleStyle::Get().GetBrush("Black .6"))
 						.Padding(0)
 						[
 							SNew(SScaleBox)
diff --git a/Engine/Source/Developer/WidgetCarousel/Public/WidgetCarouselStyle.h b/Engine/Source/Developer/WidgetCarousel/Public/WidgetCarouselStyle.h
--- a/Engine/Source/Developer/WidgetCarousel/Public/WidgetCarouselStyle.h
+++ b/Engine/Source/Developer/WidgetCarousel/Public/WidgetCarouselStyle.h
@@ -3,9 +3,72 @@
 #pragma once
 
 #include "SlateBasics.h"
+#include "WidgetCarouselStyle.generated.h"
+
+USTRUCT()
+struct WIDGETCAROUSEL_API FWidgetCarouselNavigationButtonStyle : public FSlateWidgetStyle
+{
+	GENERATED_USTRUCT_BODY()
+
+	FWidgetCarouselNavigationButtonStyle();
+
+	virtual ~FWidgetCarouselNavigationButtonStyle() {}
+
+	virtual void GetResources(TArray<const FSlateBrush*> & OutBrushes) const override;
+
+	static const FName TypeName;
+	virtual const FName GetTypeName() const override { return TypeName; };
+
+	static const FWidgetCarouselNavigationButtonStyle& GetDefault();
+
+	UPROPERTY()
+	FButtonStyle InnerButtonStyle;
+	FWidgetCarouselNavigationButtonStyle& SetInnerButtonStyle(const FButtonStyle& InInnerButtonStyle) { InnerButtonStyle = InInnerButtonStyle; return *this; }
+
+	UPROPERTY()
+	FSlateBrush NavigationButtonLeftImage;
+	FWidgetCarouselNavigationButtonStyle& SetNavigationButtonLeftImage(const FSlateBrush& InNavigationButtonLeftImage) { NavigationButtonLeftImage = InNavigationButtonLeftImage; return *this; }
+
+	UPROPERTY()
+	FSlateBrush NavigationButtonRightImage;
+	FWidgetCarouselNavigationButtonStyle& SetNavigationButtonRightImage(const FSlateBrush& InNavigationButtonRightImage) { NavigationButtonRightImage = InNavigationButtonRightImage; return *this; }
+};
+
+USTRUCT()
+struct WIDGETCAROUSEL_API FWidgetCarouselNavigationBarStyle : public FSlateWidgetStyle
+{
+	GENERATED_USTRUCT_BODY()
+
+	FWidgetCarouselNavigationBarStyle();
+
+	virtual ~FWidgetCarouselNavigationBarStyle() {}
+
+	virtual void GetResources(TArray<const FSlateBrush*> & OutBrushes) const override;
+
+	static const FName TypeName;
+	virtual const FName GetTypeName() const override { return TypeName; };
+
+	static const FWidgetCarouselNavigationBarStyle& GetDefault();
+
+	UPROPERTY()
+	FSlateBrush HighlightBrush;
+	FWidgetCarouselNavigationBarStyle& SetHighlightBrush(const FSlateBrush& InHighlightBrush) { HighlightBrush = InHighlightBrush; return *this; }
+
+	UPROPERTY()
+	FButtonStyle LeftButtonStyle;
+	FWidgetCarouselNavigationBarStyle& SetLeftButtonStyle(const FButtonStyle& InLeftButtonStyle) { LeftButtonStyle = InLeftButtonStyle; return *this; }
+
+	UPROPERTY()
+	FButtonStyle CenterButtonStyle;
+	FWidgetCarouselNavigationBarStyle& SetCenterButtonStyle(const FButtonStyle& InCenterButtonStyle) { CenterButtonStyle = InCenterButtonStyle; return *this; }
+
+	UPROPERTY()
+	FButtonStyle RightButtonStyle;
+	FWidgetCarouselNavigationBarStyle& SetRightButtonStyle(const FButtonStyle& InRightButtonStyle) { RightButtonStyle = InRightButtonStyle; return *this; }
+};
 
 /** Manages the style which provides resources for the add content dialog. */
-class WIDGETCAROUSEL_API FWidgetCarouselStyle
+class WIDGETCAROUSEL_API FWidgetCarouselModuleStyle
 {
 public:
 
diff --git a/Engine/Source/Developer/WidgetCarousel/WidgetCarousel.Build.cs b/Engine/Source/Developer/WidgetCarousel/WidgetCarousel.Build.cs
--- a/Engine/Source/Developer/WidgetCarousel/WidgetCarousel.Build.cs
+++ b/Engine/Source/Developer/WidgetCarousel/WidgetCarousel.Build.cs
@@ -12,7 +12,7 @@ public class WidgetCarousel : ModuleRules
 				"Slate",
 				"SlateCore",
 				"InputCore",
-				"EditorStyle",
+				"CoreUObject"
 			}
 		);
 
