The WidgetCarousel no longer defaults to the WidgetCarouselModuleStyle. Instead this must be specified upon construction. This prevents warning spam for applications like the launcher which don't ship editor content which the WidgetCarouselModuleStyle is dependent upon.

Merging using UE4-To-UE4-LauncherDev

diff --git a/Engine/Source/Developer/WidgetCarousel/Private/WidgetCarouselModule.cpp b/Engine/Source/Developer/WidgetCarousel/Private/WidgetCarouselModule.cpp
--- a/Engine/Source/Developer/WidgetCarousel/Private/WidgetCarouselModule.cpp
+++ b/Engine/Source/Developer/WidgetCarousel/Private/WidgetCarouselModule.cpp
@@ -4,12 +4,12 @@
 
 void FWidgetCarouselModule::StartupModule()
 {
-	FWidgetCarouselModuleStyle::Initialize();
+
 }
 
 void FWidgetCarouselModule::ShutdownModule()
 {
-	FWidgetCarouselModuleStyle::Shutdown();
+
 }
 
 IMPLEMENT_MODULE(FWidgetCarouselModule, WidgetCarousel);
diff --git a/Engine/Source/Developer/WidgetCarousel/Public/SCarouselNavigationBar.h b/Engine/Source/Developer/WidgetCarousel/Public/SCarouselNavigationBar.h
--- a/Engine/Source/Developer/WidgetCarousel/Public/SCarouselNavigationBar.h
+++ b/Engine/Source/Developer/WidgetCarousel/Public/SCarouselNavigationBar.h
@@ -12,7 +12,7 @@ class WIDGETCAROUSEL_API SCarouselNavigationBar : public SCompoundWidget
 
 public:
 	SLATE_BEGIN_ARGS(SCarouselNavigationBar)
-		: _Style(&FWidgetCarouselModuleStyle::Get().GetWidgetStyle<FWidgetCarouselNavigationBarStyle>("CarouselNavigationBar")) 
+		: _Style() 
 		, _OnSelectedIndexChanged()
 		, _ItemCount(0)
 		, _CurrentItemIndex(0)
diff --git a/Engine/Source/Developer/WidgetCarousel/Public/SCarouselNavigationButton.h b/Engine/Source/Developer/WidgetCarousel/Public/SCarouselNavigationButton.h
--- a/Engine/Source/Developer/WidgetCarousel/Public/SCarouselNavigationButton.h
+++ b/Engine/Source/Developer/WidgetCarousel/Public/SCarouselNavigationButton.h
@@ -95,7 +95,7 @@ public:
 	};
 
 	SLATE_BEGIN_ARGS(SCarouselNavigationButton)
-		: _Style(&FWidgetCarouselModuleStyle::Get().GetWidgetStyle<FWidgetCarouselNavigationButtonStyle>("CarouselNavigationButton"))
+		: _Style()
 		{}
 
 		SLATE_STYLE_ARGUMENT(FWidgetCarouselNavigationButtonStyle, Style)
diff --git a/Engine/Source/Developer/WidgetCarousel/Public/SWidgetCarouselWithNavigation.h b/Engine/Source/Developer/WidgetCarousel/Public/SWidgetCarouselWithNavigation.h
--- a/Engine/Source/Developer/WidgetCarousel/Public/SWidgetCarouselWithNavigation.h
+++ b/Engine/Source/Developer/WidgetCarousel/Public/SWidgetCarouselWithNavigation.h
@@ -17,9 +17,15 @@ public:
 	typedef typename TSlateDelegates< ItemType >::FOnGenerateWidget FOnGenerateWidget;
 
 	SLATE_BEGIN_ARGS(SWidgetCarouselWithNavigation<ItemType>)
-		: _WidgetItemsSource(static_cast<const TArray<ItemType>*>(NULL))
+		: _NavigationBarStyle()
+		, _NavigationButtonStyle()
+		, _WidgetItemsSource(static_cast<const TArray<ItemType>*>(NULL))
 	{ }
 
+		SLATE_STYLE_ARGUMENT(FWidgetCarouselNavigationBarStyle, NavigationBarStyle)
+
+		SLATE_STYLE_ARGUMENT(FWidgetCarouselNavigationButtonStyle, NavigationButtonStyle)
+
 		/** Called when we change a widget */
 		SLATE_EVENT(FOnGenerateWidget, OnGenerateWidget)
 
@@ -76,6 +82,7 @@ public:
 						.HAlign(HAlign_Left)
 						[
 							SNew(SCarouselNavigationButton)
+							.Style(InArgs._NavigationButtonStyle)
 							.OnClicked(this, &SWidgetCarouselWithNavigation::HandleNextButtonClicked)
 							.Visibility(this, &SWidgetCarouselWithNavigation::GetScreenshotNavigationVisibility)
 							.OnBeginPeek(this, &SWidgetCarouselWithNavigation::HandlePeak, EWidgetCarouselScrollDirection::Carousel_Right)
@@ -86,6 +93,7 @@ public:
 						.HAlign(HAlign_Right)
 						[
 							SNew(SCarouselNavigationButton)
+							.Style(InArgs._NavigationButtonStyle)
 							.OnClicked(this, &SWidgetCarouselWithNavigation::HandlePreviousButtonClicked)
 							.Visibility(this, &SWidgetCarouselWithNavigation::GetScreenshotNavigationVisibility)
 							.OnBeginPeek(this, &SWidgetCarouselWithNavigation::HandlePeak, EWidgetCarouselScrollDirection::Carousel_Left)
@@ -101,6 +109,7 @@ public:
 			.Padding(10.0f, 10.f)
 			[
 				SNew(SCarouselNavigationBar)
+				.Style(InArgs._NavigationBarStyle)
 				.Visibility(this, &SWidgetCarouselWithNavigation::GetScreenshotNavigationVisibility)
 				.ItemCount(WidgetItemsSource->Num())
 				.CurrentItemIndex(this, &SWidgetCarouselWithNavigation::GetCurrentItemIndex)
diff --git a/Engine/Source/Editor/AddContentDialog/Private/AddContentDialogModule.cpp b/Engine/Source/Editor/AddContentDialog/Private/AddContentDialogModule.cpp
--- a/Engine/Source/Editor/AddContentDialog/Private/AddContentDialogModule.cpp
+++ b/Engine/Source/Editor/AddContentDialog/Private/AddContentDialogModule.cpp
@@ -7,6 +7,7 @@
 #include "ModuleManager.h"
 #include "SDockTab.h"
 #include "WidgetCarouselModule.h"
+#include "WidgetCarouselStyle.h"
 
 #define LOCTEXT_NAMESPACE "AddContentDialog"
 
@@ -17,6 +18,8 @@ public:
 	virtual void StartupModule() override
 	{
 		FModuleManager::LoadModuleChecked<FWidgetCarouselModule>("WidgetCarousel");
+		FWidgetCarouselModuleStyle::Initialize();
+
 		ContentSourceProviderManager = TSharedPtr<FContentSourceProviderManager>(new FContentSourceProviderManager());
 		ContentSourceProviderManager->RegisterContentSourceProvider(MakeShareable(new FFeaturePackContentSourceProvider()));
 		FAddContentDialogStyle::Initialize();
@@ -25,6 +28,7 @@ public:
 	virtual void ShutdownModule() override
 	{
 		FAddContentDialogStyle::Shutdown();
+		FWidgetCarouselModuleStyle::Shutdown();
 	}
 
 	virtual TSharedRef<FContentSourceProviderManager> GetContentSourceProviderManager() override
diff --git a/Engine/Source/Editor/AddContentDialog/Private/SAddContentWidget.cpp b/Engine/Source/Editor/AddContentDialog/Private/SAddContentWidget.cpp
--- a/Engine/Source/Editor/AddContentDialog/Private/SAddContentWidget.cpp
+++ b/Engine/Source/Editor/AddContentDialog/Private/SAddContentWidget.cpp
@@ -282,8 +282,10 @@ TSharedRef<SWidget> SAddContentWidget::CreateContentSourceDetail(TSharedPtr<FCon
 TSharedRef<SWidget> SAddContentWidget::CreateScreenshotCarousel(TSharedPtr<FContentSourceViewModel> ContentSource)
 {
 	return SNew(SWidgetCarouselWithNavigation<TSharedPtr<FSlateBrush>>)
-	.OnGenerateWidget(this, &SAddContentWidget::CreateScreenshotWidget)
-	.WidgetItemsSource(ContentSource->GetScreenshotBrushes());
+		.NavigationBarStyle(FWidgetCarouselModuleStyle::Get(), "CarouselNavigationButton")
+		.NavigationButtonStyle(FWidgetCarouselModuleStyle::Get(), "CarouselNavigationBar")
+		.OnGenerateWidget(this, &SAddContentWidget::CreateScreenshotWidget)
+		.WidgetItemsSource(ContentSource->GetScreenshotBrushes());
 }
 
 void SAddContentWidget::CategoryCheckBoxCheckStateChanged(ECheckBoxState CheckState, FCategoryViewModel Category)
